Fix for CIS error
#rnx

diff --git a/Engine/Plugins/Runtime/MeshReconstruction/DummyMeshReconstructor/Source/DummyMeshReconstructor/Private/DummyMeshReconstructor.cpp b/Engine/Plugins/Runtime/MeshReconstruction/DummyMeshReconstructor/Source/DummyMeshReconstructor/Private/DummyMeshReconstructor.cpp
--- a/Engine/Plugins/Runtime/MeshReconstruction/DummyMeshReconstructor/Source/DummyMeshReconstructor/Private/DummyMeshReconstructor.cpp
+++ b/Engine/Plugins/Runtime/MeshReconstruction/DummyMeshReconstructor/Source/DummyMeshReconstructor/Private/DummyMeshReconstructor.cpp
@@ -154,7 +154,7 @@ private:
 		static const int32 MAX_BOXES = 20;
 		static const FBox RANDOM_SIZE_BOX = FBox(FVector::ZeroVector, 0.25f*BRICK_SIZE);
 
-		const int32 NumBoxes = FMath::FloorToInt(FMath::RandRange(MIN_BOXES, MAX_BOXES));
+		const int32 NumBoxes = FMath::FloorToInt((float)FMath::RandRange(MIN_BOXES, MAX_BOXES));
 		const int32 NumUniqueVerts = NumBoxes * 8;
 		const int32 NumTris = NumBoxes * 6 * 2; // 2 tris per box face
 		const int32 NumVertIndices = 3 * NumTris;
@@ -163,9 +163,9 @@ private:
 		const IMRMesh::FBrickId BrickId = ++NextBrickId;
 		
 		const FIntVector BrickCoords(
-			FMath::FloorToInt(FMath::RandRange(0, BRICK_COORD_RANDMAX)),
-			FMath::FloorToInt(FMath::RandRange(0, BRICK_COORD_RANDMAX)),
-			FMath::FloorToInt(FMath::RandRange(0, BRICK_COORD_RANDMAX))
+			FMath::FloorToInt((float)FMath::RandRange(0, BRICK_COORD_RANDMAX)),
+			FMath::FloorToInt((float)FMath::RandRange(0, BRICK_COORD_RANDMAX)),
+			FMath::FloorToInt((float)FMath::RandRange(0, BRICK_COORD_RANDMAX))
 		);
 
 		const FVector BrickOrigin(BRICK_SIZE.X*BrickCoords.X, BRICK_SIZE.Y*BrickCoords.Y, BRICK_SIZE.Z*BrickCoords.Z);
